Set frame delay in the Timecode panel.

diff --git a/Engine/Source/Editor/TimeManagementEditor/Private/STimecodeProvider.cpp b/Engine/Source/Editor/TimeManagementEditor/Private/STimecodeProvider.cpp
--- a/Engine/Source/Editor/TimeManagementEditor/Private/STimecodeProvider.cpp
+++ b/Engine/Source/Editor/TimeManagementEditor/Private/STimecodeProvider.cpp
@@ -29,6 +29,10 @@ void STimecodeProvider::Construct(const FArguments& InArgs)
 			{
 				return TimecodeProviderPtr->GetFrameRate().ToPrettyText();
 			}
+			else if (GEngine->bGenerateDefaultTimecode)
+			{
+				return GEngine->GenerateDefaultTimecodeFrameRate.ToPrettyText();
+			}
 			return NSLOCTEXT("TimecodeProvider", "Undefined", "<Undefined>");
 		}))
 		.Font(InArgs._TimecodeProviderFont)
@@ -61,6 +65,10 @@ void STimecodeProvider::Construct(const FArguments& InArgs)
 					{
 						return FText::FromName(TimecodeProviderPtr->GetFName());
 					}
+					else if (GEngine->bGenerateDefaultTimecode)
+					{
+						return NSLOCTEXT("TimecodeProvider", "EngineDefault", "Engine Default");
+					}
 					return NSLOCTEXT("TimecodeProvider", "Undefined", "<Undefined>");
 				}))
 				.Font(InArgs._TimecodeProviderFont)
@@ -123,6 +131,10 @@ FSlateColor STimecodeProvider::HandleIconColorAndOpacity() const
 			return FLinearColor::Yellow;
 		}
 	}
+	else if (GEngine->bGenerateDefaultTimecode)
+	{
+		return FLinearColor::Green;
+	}
 	return FSlateColor::UseForeground();
 }
 
@@ -142,6 +154,10 @@ FText STimecodeProvider::HandleStateText() const
 			return FEditorFontGlyphs::Hourglass_O;
 		}
 	}
+	else if (GEngine->bGenerateDefaultTimecode)
+	{
+		return FEditorFontGlyphs::Clock_O;
+	}
 
 	return FEditorFontGlyphs::Exclamation;
 }
diff --git a/Engine/Source/Editor/TimeManagementEditor/Private/STimecodeProviderTab.cpp b/Engine/Source/Editor/TimeManagementEditor/Private/STimecodeProviderTab.cpp
--- a/Engine/Source/Editor/TimeManagementEditor/Private/STimecodeProviderTab.cpp
+++ b/Engine/Source/Editor/TimeManagementEditor/Private/STimecodeProviderTab.cpp
@@ -11,9 +11,11 @@
 #include "Framework/MultiBox/SToolBarComboButtonBlock.h"
 #include "LevelEditor.h"
 #include "Modules/ModuleManager.h"
+#include "ScopedTransaction.h"
 #include "STimecodeProvider.h"
 #include "Widgets/Docking/SDockTab.h"
 #include "Widgets/Input/SComboButton.h"
+#include "Widgets/Input/SSpinBox.h"
 #include "WorkspaceMenuStructure.h"
 #include "WorkspaceMenuStructureModule.h"
 
@@ -91,7 +93,7 @@ void STimecodeProviderTab::Construct(const FArguments& InArgs)
 		]
 		.OnGetMenuContent(this, &STimecodeProviderTab::OnGetMenuContent);
 
-	ButtonContent->SetEnabled(MakeAttributeLambda([] { return (GEngine && GEngine->GetTimecodeProvider() != nullptr); }));
+	ButtonContent->SetEnabled(MakeAttributeLambda([] { return GEngine && (GEngine->GetTimecodeProvider() != nullptr || GEngine->bGenerateDefaultTimecode); }));
 
 	ChildSlot
 	[
@@ -146,7 +148,50 @@ TSharedRef<SWidget> STimecodeProviderTab::OnGetMenuContent()
 		MenuBuilder.EndSection();
 	}
 
+	if (GEngine->GetTimecodeProvider() || GEngine->bGenerateDefaultTimecode)
+	{
+		MenuBuilder.BeginSection("Settings", LOCTEXT("Settings", "Settings"));
+		{
+			TSharedRef<SWidget> RefreshDelay = SNew(SSpinBox<int32>)
+				.ToolTipText(LOCTEXT("FrameDelay_ToolTip", "Number of frames to subtract from the original timecode."))
+				.Value(this, &STimecodeProviderTab::GetFrameDelay)
+				.OnValueCommitted(this, &STimecodeProviderTab::SetFrameDelay);
+
+			MenuBuilder.AddWidget(RefreshDelay, LOCTEXT("FrameDelay", "Frame Delay"));
+		}
+		MenuBuilder.EndSection();
+	}
+
 	return MenuBuilder.MakeWidget();
 }
 
+int32 STimecodeProviderTab::GetFrameDelay() const
+{
+	if (GEngine->GetTimecodeProvider())
+	{
+		return GEngine->GetTimecodeProvider()->FrameDelay;
+	}
+	else if (GEngine->bGenerateDefaultTimecode)
+	{
+		return GEngine->GenerateDefaultTimecodeFrameDelay;
+	}
+	return 0;
+}
+
+void STimecodeProviderTab::SetFrameDelay(int32 InNewValue, ETextCommit::Type)
+{
+	if (GetFrameDelay() != InNewValue)
+	{
+		const FScopedTransaction Transaction(LOCTEXT("SetFrameDelay", "TC Frame Delay"));
+		if (GEngine->GetTimecodeProvider())
+		{
+			GEngine->GetTimecodeProvider()->FrameDelay = InNewValue;
+		}
+		else if (GEngine->bGenerateDefaultTimecode)
+		{
+			GEngine->GenerateDefaultTimecodeFrameDelay = InNewValue;
+		}
+	}
+}
+
 #undef LOCTEXT_NAMESPACE
diff --git a/Engine/Source/Editor/TimeManagementEditor/Private/STimecodeProviderTab.h b/Engine/Source/Editor/TimeManagementEditor/Private/STimecodeProviderTab.h
--- a/Engine/Source/Editor/TimeManagementEditor/Private/STimecodeProviderTab.h
+++ b/Engine/Source/Editor/TimeManagementEditor/Private/STimecodeProviderTab.h
@@ -4,6 +4,7 @@
 
 #include "CoreMinimal.h"
 #include "Misc/FrameRate.h"
+#include "Types/SlateEnums.h"
 #include "Widgets/DeclarativeSyntaxSupport.h"
 #include "Widgets/SCompoundWidget.h"
 
@@ -24,6 +25,10 @@ public:
 
 	void Construct(const FArguments& InArgs);
 
+private:
+	int32 GetFrameDelay() const;
+	void SetFrameDelay(int32, ETextCommit::Type);
+
 private:
 	TSharedRef<SWidget> OnGetMenuContent();
 };
diff --git a/Engine/Source/Editor/TimeManagementEditor/TimeManagementEditor.Build.cs b/Engine/Source/Editor/TimeManagementEditor/TimeManagementEditor.Build.cs
--- a/Engine/Source/Editor/TimeManagementEditor/TimeManagementEditor.Build.cs
+++ b/Engine/Source/Editor/TimeManagementEditor/TimeManagementEditor.Build.cs
@@ -18,7 +18,9 @@ public class TimeManagementEditor : ModuleRules
 				"CoreUObject",
 				"EditorStyle",
 				"Engine",
+				"InputCore",
 				"TimeManagement",
+				"UnrealEd",
 				"WorkspaceMenuStructure",
 			}
 		);
diff --git a/Engine/Source/Runtime/Engine/Classes/Engine/Engine.h b/Engine/Source/Runtime/Engine/Classes/Engine/Engine.h
--- a/Engine/Source/Runtime/Engine/Classes/Engine/Engine.h
+++ b/Engine/Source/Runtime/Engine/Classes/Engine/Engine.h
@@ -1371,6 +1371,10 @@ public:
 	UPROPERTY(config, EditAnywhere, Category=Timecode, meta=(EditCondition="bGenerateDefaultTimecode"))
 	FFrameRate GenerateDefaultTimecodeFrameRate;
 
+	/** Number of frames to subtract from generated default timecode. */
+	UPROPERTY(AdvancedDisplay, config, EditAnywhere, Category=Timecode, meta=(EditCondition="bGenerateDefaultTimecode"))
+	int32 GenerateDefaultTimecodeFrameDelay;
+
 public:
 	/** 
 	 * Whether we should check for more than N pawns spawning in a single frame.
diff --git a/Engine/Source/Runtime/Engine/Private/GameEngine.cpp b/Engine/Source/Runtime/Engine/Private/GameEngine.cpp
--- a/Engine/Source/Runtime/Engine/Private/GameEngine.cpp
+++ b/Engine/Source/Runtime/Engine/Private/GameEngine.cpp
@@ -723,7 +723,8 @@ UEngine::UEngine(const FObjectInitializer& ObjectInitializer)
 	bIsVanillaProduct = false;
 
 	bGenerateDefaultTimecode = true;
-	GenerateDefaultTimecodeFrameRate = FFrameRate(1, 24);
+	GenerateDefaultTimecodeFrameRate = FFrameRate(24, 1);
+	GenerateDefaultTimecodeFrameDelay = 0;
 
 	GameScreenshotSaveDirectory.Path = FPaths::ScreenShotDir();
 
diff --git a/Engine/Source/Runtime/Engine/Private/UnrealEngine.cpp b/Engine/Source/Runtime/Engine/Private/UnrealEngine.cpp
--- a/Engine/Source/Runtime/Engine/Private/UnrealEngine.cpp
+++ b/Engine/Source/Runtime/Engine/Private/UnrealEngine.cpp
@@ -2248,11 +2248,13 @@ void UEngine::UpdateTimecode()
 	else if(bGenerateDefaultTimecode)
 	{
 #if PLATFORM_DESKTOP
-		FApp::SetCurrentFrameTime(FQualifiedFrameTime(USystemTimeTimecodeProvider::GenerateTimecodeFromSystemTime(GenerateDefaultTimecodeFrameRate), GenerateDefaultTimecodeFrameRate));
+		FQualifiedFrameTime NewFrameTime = FQualifiedFrameTime(USystemTimeTimecodeProvider::GenerateTimecodeFromSystemTime(GenerateDefaultTimecodeFrameRate), GenerateDefaultTimecodeFrameRate);
 #else //PLATFORM_DESKTOP
 		//If a user wish to have an accurate TC value on console, he should set an accurate TC provider in his project settings.
-		FApp::SetCurrentFrameTime(FQualifiedFrameTime(USystemTimeTimecodeProvider::GenerateTimecodeFromHighPerformanceClock(GenerateDefaultTimecodeFrameRate), GenerateDefaultTimecodeFrameRate));
+		FQualifiedFrameTime NewFrameTime = FQualifiedFrameTime(USystemTimeTimecodeProvider::GenerateTimecodeFromHighPerformanceClock(GenerateDefaultTimecodeFrameRate), GenerateDefaultTimecodeFrameRate);
 #endif
+		NewFrameTime.Time.FrameNumber -= GenerateDefaultTimecodeFrameDelay;
+		FApp::SetCurrentFrameTime(NewFrameTime);
 	}
 }
 
