Update STimecodeProviderTab to support subframe editing.

#rnx

diff --git a/Engine/Source/Editor/TimeManagementEditor/Private/STimecodeProviderTab.cpp b/Engine/Source/Editor/TimeManagementEditor/Private/STimecodeProviderTab.cpp
--- a/Engine/Source/Editor/TimeManagementEditor/Private/STimecodeProviderTab.cpp
+++ b/Engine/Source/Editor/TimeManagementEditor/Private/STimecodeProviderTab.cpp
@@ -152,7 +152,7 @@ TSharedRef<SWidget> STimecodeProviderTab::OnGetMenuContent()
 	{
 		MenuBuilder.BeginSection("Settings", LOCTEXT("Settings", "Settings"));
 		{
-			TSharedRef<SWidget> RefreshDelay = SNew(SSpinBox<int32>)
+			TSharedRef<SWidget> RefreshDelay = SNew(SSpinBox<float>)
 				.ToolTipText(LOCTEXT("FrameDelay_ToolTip", "Number of frames to subtract from the original timecode."))
 				.Value(this, &STimecodeProviderTab::GetFrameDelay)
 				.OnValueCommitted(this, &STimecodeProviderTab::SetFrameDelay);
@@ -165,7 +165,7 @@ TSharedRef<SWidget> STimecodeProviderTab::OnGetMenuContent()
 	return MenuBuilder.MakeWidget();
 }
 
-int32 STimecodeProviderTab::GetFrameDelay() const
+float STimecodeProviderTab::GetFrameDelay() const
 {
 	if (GEngine->GetTimecodeProvider())
 	{
@@ -175,10 +175,10 @@ int32 STimecodeProviderTab::GetFrameDelay() const
 	{
 		return GEngine->GenerateDefaultTimecodeFrameDelay;
 	}
-	return 0;
+	return 0.f;
 }
 
-void STimecodeProviderTab::SetFrameDelay(int32 InNewValue, ETextCommit::Type)
+void STimecodeProviderTab::SetFrameDelay(float InNewValue, ETextCommit::Type)
 {
 	if (GetFrameDelay() != InNewValue)
 	{
diff --git a/Engine/Source/Editor/TimeManagementEditor/Private/STimecodeProviderTab.h b/Engine/Source/Editor/TimeManagementEditor/Private/STimecodeProviderTab.h
--- a/Engine/Source/Editor/TimeManagementEditor/Private/STimecodeProviderTab.h
+++ b/Engine/Source/Editor/TimeManagementEditor/Private/STimecodeProviderTab.h
@@ -26,8 +26,8 @@ public:
 	void Construct(const FArguments& InArgs);
 
 private:
-	int32 GetFrameDelay() const;
-	void SetFrameDelay(int32, ETextCommit::Type);
+	float GetFrameDelay() const;
+	void SetFrameDelay(float, ETextCommit::Type);
 
 private:
 	TSharedRef<SWidget> OnGetMenuContent();
diff --git a/Engine/Source/Runtime/Engine/Private/GameEngine.cpp b/Engine/Source/Runtime/Engine/Private/GameEngine.cpp
--- a/Engine/Source/Runtime/Engine/Private/GameEngine.cpp
+++ b/Engine/Source/Runtime/Engine/Private/GameEngine.cpp
@@ -724,7 +724,7 @@ UEngine::UEngine(const FObjectInitializer& ObjectInitializer)
 
 	bGenerateDefaultTimecode = true;
 	GenerateDefaultTimecodeFrameRate = FFrameRate(24, 1);
-	GenerateDefaultTimecodeFrameDelay = 0;
+	GenerateDefaultTimecodeFrameDelay = 0.f;
 
 	GameScreenshotSaveDirectory.Path = FPaths::ScreenShotDir();
 
